Fix the FTraceFilter transactional bug by introducing a new transactionally safe critical section and scope lock for it to use.

I also worked out how to add a new AutoRTFMEngineTests target that lets us test a much greater surface area of the engine.

diff --git a/Engine/Source/Programs/AutoRTFMEngineTests/AutoRTFMEngineTests.Build.cs b/Engine/Source/Programs/AutoRTFMEngineTests/AutoRTFMEngineTests.Build.cs
new file mode 100644
--- /dev/null
+++ b/Engine/Source/Programs/AutoRTFMEngineTests/AutoRTFMEngineTests.Build.cs
@@ -0,0 +1,37 @@
+// Copyright Epic Games, Inc. All Rights Reserved.
+
+using UnrealBuildTool;
+
+public class AutoRTFMEngineTests : ModuleRules
+{
+	public AutoRTFMEngineTests(ReadOnlyTargetRules Target) : base(Target)
+	{
+		PublicIncludePathModuleNames.Add("Launch");
+		PrivateIncludePathModuleNames.Add("DerivedDataCache");
+
+		PrivateDependencyModuleNames.AddRange(
+			new string[]
+			{
+ 				"ApplicationCore",
+ 				"AutomationController",
+ 				"AutomationWorker",
+                "Core",
+				"Projects",
+				"Engine",
+				"HeadMountedDisplay",
+				"InstallBundleManager",
+                "MediaUtils",
+				"MRMesh",
+				"MoviePlayer",
+				"MoviePlayerProxy",
+				"PreLoadScreen",
+				"ProfilerService",
+				"ReplicationSystemTestPlugin",
+				"SessionServices",
+				"SlateNullRenderer",
+				"SlateRHIRenderer",
+				"ProfileVisualizer",
+ 			}
+		);
+	}
+}
diff --git a/Engine/Source/Programs/AutoRTFMEngineTests/AutoRTFMEngineTests.Target.cs b/Engine/Source/Programs/AutoRTFMEngineTests/AutoRTFMEngineTests.Target.cs
new file mode 100644
--- /dev/null
+++ b/Engine/Source/Programs/AutoRTFMEngineTests/AutoRTFMEngineTests.Target.cs
@@ -0,0 +1,51 @@
+// Copyright Epic Games, Inc. All Rights Reserved.
+
+using UnrealBuildTool;
+using System.Collections.Generic;
+
+[SupportedPlatforms(UnrealPlatformClass.All)]
+public class AutoRTFMEngineTestsTarget : TargetRules
+{
+	public AutoRTFMEngineTestsTarget(TargetInfo Target) : base(Target)
+	{
+		Type = TargetType.Program;
+		LinkType = TargetLinkType.Monolithic;
+		LaunchModuleName = "AutoRTFMEngineTests";
+
+		// No editor-only data is needed
+		bBuildWithEditorOnlyData = false;
+
+		bCompileAgainstEngine = true;
+		bCompileAgainstCoreUObject = false;
+		bCompileAgainstApplicationCore = true;
+        bCompileWithPluginSupport = true;
+        bBuildDeveloperTools = false;
+        bBuildRequiresCookedData = true; // this program requires no data
+
+		// No ICU internationalization as it causes shutdown errors.
+		bCompileICU = false;
+
+		// Don't need slate
+		bUsesSlate = false;
+
+		IncludeOrderVersion = EngineIncludeOrderVersion.Latest;
+		bIsBuildingConsoleApplication = true;
+		bLegalToDistributeBinary = true;
+
+		// Always want logging so we can see the test results
+        bUseLoggingInShipping = true;
+
+		// Network config
+		bWithPushModel = true;		
+		bUseIris = true;
+
+		// Need to force enable tracing so we can test against it.
+		bForceEnableTrace = true;
+
+		// Load time profiling brings object construction to a crawl.
+		GlobalDefinitions.Add("LOADTIMEPROFILERTRACE_ENABLED=0");
+
+		bUseAutoRTFMCompiler = true;
+	}
+}
+
diff --git a/Engine/Source/Programs/AutoRTFMEngineTests/Private/AutoRTFMEngineTests.cpp b/Engine/Source/Programs/AutoRTFMEngineTests/Private/AutoRTFMEngineTests.cpp
new file mode 100644
--- /dev/null
+++ b/Engine/Source/Programs/AutoRTFMEngineTests/Private/AutoRTFMEngineTests.cpp
@@ -0,0 +1,215 @@
+// Copyright Epic Games, Inc. All Rights Reserved.
+
+#include "AutoRTFMEngineTests.h"
+#include "RequiredProgramMainCPPInclude.h"
+#include "Styling/UMGCoreStyle.h"
+#include "AutomationTestRunner.h"
+#include "NullTestRunner.h"
+#include "Materials/Material.h"
+#include "AutoRTFM/AutoRTFM.h"
+
+#if WITH_AUTOMATION_WORKER
+namespace UE::AutoRTFM
+{
+	class FAutomationTestRunner;
+}
+
+typedef UE::AutoRTFM::FAutomationTestRunner TestRunner;
+#else
+typedef FNullTestRunner TestRunner;
+#endif
+
+DEFINE_LOG_CATEGORY(LogAutoRTFMEngineTests);
+
+IMPLEMENT_APPLICATION(AutoRTFMEngineTests, "");
+
+static void PreInit();
+static void LoadModules();
+static void PostInit();
+static void TearDown();
+
+#if defined(PLATFORM_XBOXONE) && PLATFORM_XBOXONE
+int TestMain()
+#else
+INT32_MAIN_INT32_ARGC_TCHAR_ARGV()
+#endif
+{
+	FString OriginalCmdLine;
+
+#if !(defined(PLATFORM_XBOXONE) && PLATFORM_XBOXONE)
+	// Parse original cmdline if there is one
+	OriginalCmdLine = FCommandLine::BuildFromArgV(nullptr, ArgC, ArgV, nullptr);	
+#endif
+
+	// Due to some code not respecting nullrhi etc.
+	PRIVATE_GIsRunningCommandlet = true;
+	PRIVATE_GAllowCommandletRendering = false;
+	PRIVATE_GAllowCommandletAudio = false;
+
+	// Init cmd line used for test
+	{
+		FString CmdLineOverride(TEXT("-nullrhi -NoAsyncLoadingThread -NoAsyncPostLoad -noedl -unattended -LogCmds=\"LogSlate off, LogSlateStyle off, LogUObjectBase off, LogUObjectGlobals off, LogConsoleResponse off, LogPackageLocalizationManager off, LogStreaming off, LogCsvProfiler off, LogDeviceProfileManager off, LogConfig off, AutoRTFMEngineTests on\""));
+		FCommandLine::Set(ToCStr(CmdLineOverride));
+	}
+
+	FLogSuppressionInterface::Get().ProcessConfigAndCommandLine();
+
+	PreInit();
+	LoadModules();
+	PostInit();
+
+	UE_LOG(LogAutoRTFMEngineTests, Display, TEXT("AutoRTFMEngineTests"));
+
+	bool TestsPassed = false;
+
+	{
+		TUniquePtr<TestRunner> Runner(new TestRunner());
+
+		FString TestFilter;
+
+		if (FParse::Value(ToCStr(OriginalCmdLine), TEXT("TestFilter="), TestFilter))
+		{
+			TestsPassed = Runner->RunTests(ToCStr(TestFilter));
+		}
+		else
+		{
+			TestsPassed = Runner->RunTests();
+		}
+		
+	}
+
+	TearDown();
+
+	return TestsPassed ? 0 : 1;
+}
+
+static void TearDown()
+{
+	RequestEngineExit(TEXT("Shutting down AutoRTFMEngineTests"));
+
+	FPlatformApplicationMisc::TearDown();
+	FPlatformMisc::PlatformTearDown();
+
+	if (GLog)
+	{
+		GLog->TearDown();
+	}
+
+	FCoreDelegates::OnExit.Broadcast();
+	FModuleManager::Get().UnloadModulesAtShutdown();
+
+#if STATS
+	FThreadStats::StopThread();
+#endif
+
+	FTaskGraphInterface::Shutdown();
+
+	if (GConfig)
+	{
+		GConfig->Exit();
+		delete GConfig;
+		GConfig = nullptr;
+	}
+
+	FTraceAuxiliary::Shutdown();
+}
+
+static void PreInit()
+{
+	// We enable the AutoRTFM runtime as the tests depend on it.
+	AutoRTFM::ForTheRuntime::SetAutoRTFMRuntime(AutoRTFM::ForTheRuntime::AutoRTFM_Enabled);
+
+	FGenericPlatformOutputDevices::SetupOutputDevices();
+
+	GError = FPlatformApplicationMisc::GetErrorOutputDevice();
+	GWarn = FPlatformApplicationMisc::GetFeedbackContext();
+
+	FPlatformMisc::PlatformInit();
+#if WITH_APPLICATION_CORE
+	FPlatformApplicationMisc::Init();
+#endif
+	FPlatformMemory::Init();
+
+	// Initialize trace
+	FString Parameter;
+	if (FParse::Value(FCommandLine::Get(), TEXT("-trace="), Parameter, false))
+	{
+		FTraceAuxiliary::Initialize(FCommandLine::Get());
+		FTraceAuxiliary::TryAutoConnect();
+	}
+
+	if (IPlatformFile* WrapperFile = FPlatformFileManager::Get().GetPlatformFile(TEXT("AutoRTFMEngineTestsFile")))
+	{
+		IPlatformFile* CurrentPlatformFile = &FPlatformFileManager::Get().GetPlatformFile();
+		WrapperFile->Initialize(CurrentPlatformFile, TEXT(""));
+		FPlatformFileManager::Get().SetPlatformFile(*WrapperFile);
+	}
+
+#if WITH_COREUOBJECT
+	// Initialize the PackageResourceManager, which is needed to load any (non-script) Packages. It is first used in ProcessNewlyLoadedObjects (due to the loading of asset references in Class Default Objects)
+	// It has to be intialized after the AssetRegistryModule; the editor implementations of PackageResourceManager relies on it
+	IPackageResourceManager::Initialize();
+#endif
+
+	FDelayedAutoRegisterHelper::RunAndClearDelayedAutoRegisterDelegates(EDelayedRegisterRunPhase::FileSystemReady);
+
+	FConfigCacheIni::InitializeConfigSystem();
+
+	// Config overrides
+	GConfig->SetInt(TEXT("/Script/Engine.GarbageCollectionSettings"), TEXT("gc.MaxObjectsNotConsideredByGC"), 0, GEngineIni);
+	GConfig->SetString(TEXT("/Script/Engine.Engine"), TEXT("AIControllerClassName"), TEXT("/Script/AIModule.AIController"), GEngineIni);
+	GConfig->SetString(TEXT("/Script/Engine.Engine"), TEXT("DefaultMaterialName"), TEXT("/Engine/Transient.MockDefaultMaterial"), GEngineIni);
+	GConfig->SetString(TEXT("/Script/Engine.Engine"), TEXT("DefaultLightFunctionMaterialName"), TEXT("/Engine/Transient.MockDefaultMaterial"), GEngineIni);
+	GConfig->SetString(TEXT("/Script/Engine.Engine"), TEXT("DefaultDeferredDecalMaterialName"), TEXT("/Engine/Transient.MockDefaultMaterial"), GEngineIni);
+	GConfig->SetString(TEXT("/Script/Engine.Engine"), TEXT("DefaultPostProcessMaterialName"), TEXT("/Engine/Transient.MockDefaultMaterial"), GEngineIni);
+
+	GGameThreadId = FPlatformTLS::GetCurrentThreadId();
+	FTaskGraphInterface::Startup(FPlatformMisc::NumberOfCores());
+	FTaskGraphInterface::Get().AttachToThread(ENamedThreads::GameThread);
+
+	FDelayedAutoRegisterHelper::RunAndClearDelayedAutoRegisterDelegates(EDelayedRegisterRunPhase::TaskGraphSystemReady);
+
+#if STATS
+	FThreadStats::StartThread();
+#endif
+
+	FDelayedAutoRegisterHelper::RunAndClearDelayedAutoRegisterDelegates(EDelayedRegisterRunPhase::StatSystemReady);
+
+	FCoreStyle::ResetToDefault();
+	FUMGCoreStyle::ResetToDefault();
+}
+
+static void LoadModules()
+{
+	// Always attempt to load CoreUObject. It requires additional pre-init which is called from its module's StartupModule method.
+#if WITH_COREUOBJECT
+	// Always register the UObjects callback for VNI and general consistency with the callbacks ProcessNewlyLoadedUObjects calls.
+	RegisterProcessNewlyLoadedUObjects();
+	FModuleManager::Get().LoadModule(TEXT("CoreUObject"));
+
+	FCoreDelegates::OnInit.Broadcast();
+#endif
+
+	// Create a mock default material to keep the material system happy
+	UMaterial* MockMaterial = NewObject<UMaterial>(GetTransientPackage(), UMaterial::StaticClass(), TEXT("MockDefaultMaterial"), RF_Transient | RF_MarkAsRootSet);
+
+	// ChaosEngineSolvers requires ChaosSolvers. We're not able to call ProcessNewlyLoadedObjects before this module is loaded.
+	FModuleManager::Get().LoadModule(TEXT("ChaosSolvers"));
+	ProcessNewlyLoadedUObjects();
+
+	//FModuleManager::Get().LoadModule(TEXT("IrisCore"));
+}
+
+static void PostInit()
+{
+#if WITH_COREUOBJECT
+	// Required for GC to be allowed.
+	if (GUObjectArray.IsOpenForDisregardForGC())
+	{
+		GUObjectArray.CloseDisregardForGC();
+	}
+#endif
+
+	// Disable ini file operations
+	GConfig->DisableFileOperations();
+}
diff --git a/Engine/Source/Programs/AutoRTFMEngineTests/Private/AutoRTFMEngineTests.h b/Engine/Source/Programs/AutoRTFMEngineTests/Private/AutoRTFMEngineTests.h
new file mode 100644
--- /dev/null
+++ b/Engine/Source/Programs/AutoRTFMEngineTests/Private/AutoRTFMEngineTests.h
@@ -0,0 +1,7 @@
+// Copyright Epic Games, Inc. All Rights Reserved.
+
+#pragma once
+
+#include "CoreMinimal.h"
+
+DECLARE_LOG_CATEGORY_EXTERN(LogAutoRTFMEngineTests, Log, All)
diff --git a/Engine/Source/Programs/AutoRTFMEngineTests/Private/AutoRTFMEngineTestsFileModule.cpp b/Engine/Source/Programs/AutoRTFMEngineTests/Private/AutoRTFMEngineTestsFileModule.cpp
new file mode 100644
--- /dev/null
+++ b/Engine/Source/Programs/AutoRTFMEngineTests/Private/AutoRTFMEngineTestsFileModule.cpp
@@ -0,0 +1,59 @@
+// Copyright Epic Games, Inc. All Rights Reserved.
+
+#include "HAL/PlatformFile.h"
+#include "HAL/IPlatformFileModule.h"
+#include "Modules/ModuleInterface.h"
+#include "Modules/ModuleManager.h"
+
+class FAutoRTFMEngineTestsFile : public IPlatformFile
+{
+public:
+	FAutoRTFMEngineTestsFile() : LowerLevel(nullptr) {}
+
+	virtual bool ShouldBeUsed(IPlatformFile* Inner, const TCHAR* CmdLine) const override { return true; }
+	virtual bool Initialize(IPlatformFile* Inner, const TCHAR* CmdLine) override { LowerLevel = Inner; return true; }
+
+	virtual IPlatformFile* GetLowerLevel() { return LowerLevel; }
+	virtual void SetLowerLevel(IPlatformFile* NewLowerLevel) { LowerLevel = NewLowerLevel; }
+	virtual const TCHAR* GetName() const { return TEXT("AutoRTFMEngineTests"); }
+
+	virtual bool FileExists(const TCHAR* Filename) override { return false; }
+	virtual int64 FileSize(const TCHAR* Filename) override {return -1LL; }
+
+	virtual bool DeleteFile(const TCHAR* Filename) override { return false; }
+	virtual bool IsReadOnly(const TCHAR* Filename) override { return false; }
+
+	virtual bool MoveFile(const TCHAR* To, const TCHAR* From) override { return true; }
+	virtual bool SetReadOnly(const TCHAR* Filename, bool bNewReadOnlyValue) override { return false; }
+	virtual FDateTime GetTimeStamp(const TCHAR* Filename) override { return FDateTime::MinValue(); }
+	virtual void SetTimeStamp(const TCHAR* Filename, FDateTime DateTime) override {}
+	virtual FDateTime GetAccessTimeStamp(const TCHAR* Filename) override { return FDateTime::MinValue(); }
+	virtual FString GetFilenameOnDisk(const TCHAR* Filename) override { return TEXT("I.ris"); }
+	virtual IFileHandle* OpenRead(const TCHAR* Filename, bool bAllowWrite = false) override { return nullptr; }
+	virtual IFileHandle* OpenWrite(const TCHAR* Filename, bool bAppend = false, bool bAllowRead = false) override { return nullptr; }
+
+	virtual bool DirectoryExists(const TCHAR* Directory) override { return false; }
+
+	virtual bool CreateDirectory(const TCHAR* Directory) override { return false; }
+	virtual bool DeleteDirectory(const TCHAR* Directory) override { return false; }
+	virtual FFileStatData GetStatData(const TCHAR* FilenameOrDirectory) override { FFileStatData Data; Data.bIsValid = false; return Data; }
+
+	virtual bool IterateDirectory(const TCHAR* Directory, FDirectoryVisitor& Visitor) override { return true; }
+	virtual bool IterateDirectoryStat(const TCHAR* Directory, FDirectoryStatVisitor& Visitor) override { return true; }
+
+private:
+	IPlatformFile* LowerLevel;
+};
+
+
+class FAutoRTFMEngineTestsFileModule : public IPlatformFileModule
+{
+public:
+	virtual IPlatformFile* GetPlatformFile() override
+	{
+		static TUniquePtr<IPlatformFile> AutoDestroySingleton = MakeUnique<FAutoRTFMEngineTestsFile>();
+		return AutoDestroySingleton.Get();
+	}
+};
+
+IMPLEMENT_MODULE(FAutoRTFMEngineTestsFileModule, AutoRTFMEngineTestsFile);
diff --git a/Engine/Source/Programs/AutoRTFMEngineTests/Private/AutomationTestRunner.cpp b/Engine/Source/Programs/AutoRTFMEngineTests/Private/AutomationTestRunner.cpp
new file mode 100644
--- /dev/null
+++ b/Engine/Source/Programs/AutoRTFMEngineTests/Private/AutomationTestRunner.cpp
@@ -0,0 +1,80 @@
+// Copyright Epic Games, Inc. All Rights Reserved.
+
+#include "AutomationTestRunner.h"
+
+#if WITH_AUTOMATION_WORKER
+
+namespace UE::AutoRTFM
+{
+
+FAutomationTestRunner::FAutomationTestRunner()
+{
+}
+
+bool
+FAutomationTestRunner::RunTests(const TCHAR* TestFilter)
+{
+	constexpr int32 ExpectedTestCount = 2048;
+
+	TArray<FAutomationTestInfo> TestInfos;
+	TestInfos.Empty(ExpectedTestCount);
+
+	FAutomationTestFramework& TestFramework = FAutomationTestFramework::Get();
+	TestFramework.SetRequestedTestFilter(EAutomationTestFlags::SmokeFilter | EAutomationTestFlags::EngineFilter);
+	TestFramework.GetValidTestNames(TestInfos);
+	int TestCount = TestInfos.Num();
+
+	if (TestCount <= 0)
+	{
+		return true;
+	}
+
+	// Stack walking doesn't work properly on Windows when omitting frame pointers. See WindowsPlatformStackWalk.cpp.
+	const bool bCaptureStack = TestFramework.GetCaptureStack();
+	TestFramework.SetCaptureStack(false);
+
+	bool AllPassed = true;
+
+	const double TestStartTime = FPlatformTime::Seconds();
+
+	for (int TestIt = 0; TestIt != TestCount; ++TestIt)
+	{
+		const FAutomationTestInfo& TestInfo = TestInfos[TestIt];
+		const FString& TestCommand = TestInfo.GetTestName();
+
+		if (!TestCommand.Contains("AutoRTFM"))
+			continue;
+
+		if (TestFilter && !TestCommand.Contains(TestFilter))
+		{
+			continue;
+		}
+
+		constexpr int32 RoleIndex = 0;
+		TestFramework.StartTestByName(TestCommand, RoleIndex);
+		FAutomationTestExecutionInfo ExecutionInfo;
+
+		if (!TestFramework.StopTest(ExecutionInfo))
+		{
+			for (const FAutomationExecutionEntry& Entry : ExecutionInfo.GetEntries())
+			{
+				UE_LOG(LogAutoRTFMEngineTests, Display, TEXT("%s"), *Entry.Event.Message);
+			}
+			
+			AllPassed = false;
+		}
+	}
+
+	const double TestEndTime = FPlatformTime::Seconds();
+
+	const double TestTime = TestEndTime - TestStartTime;
+	UE_LOG(LogAutoRTFMEngineTests, Display, TEXT("Tests took %.3lf seconds to execute"), TestTime);
+
+	TestFramework.SetCaptureStack(bCaptureStack);
+
+	return AllPassed;
+}
+
+}
+
+#endif
diff --git a/Engine/Source/Programs/AutoRTFMEngineTests/Private/AutomationTestRunner.h b/Engine/Source/Programs/AutoRTFMEngineTests/Private/AutomationTestRunner.h
new file mode 100644
--- /dev/null
+++ b/Engine/Source/Programs/AutoRTFMEngineTests/Private/AutomationTestRunner.h
@@ -0,0 +1,23 @@
+// Copyright Epic Games, Inc. All Rights Reserved.
+
+#pragma once
+
+#include "Misc/AutomationTest.h"
+#include "AutoRTFMEngineTests.h"
+
+#if WITH_AUTOMATION_WORKER
+
+namespace UE::AutoRTFM
+{
+
+class FAutomationTestRunner
+{
+public:
+	FAutomationTestRunner();
+	
+	// Returns true if the tests ran successfully.
+	bool RunTests(const TCHAR* TestFilter = nullptr);
+};
+
+}
+#endif
diff --git a/Engine/Source/Programs/AutoRTFMEngineTests/Private/NullTestRunner.h b/Engine/Source/Programs/AutoRTFMEngineTests/Private/NullTestRunner.h
new file mode 100644
--- /dev/null
+++ b/Engine/Source/Programs/AutoRTFMEngineTests/Private/NullTestRunner.h
@@ -0,0 +1,11 @@
+// Copyright Epic Games, Inc. All Rights Reserved.
+
+#pragma once
+
+#include "HAL/Platform.h"
+
+class FNullTestRunner
+{
+public:
+	void RunTests(const TCHAR* Filter = nullptr) {}
+};
diff --git a/Engine/Source/Programs/UnrealBuildTool/Configuration/ReadOnlyTargetRules.cs b/Engine/Source/Programs/UnrealBuildTool/Configuration/ReadOnlyTargetRules.cs
--- a/Engine/Source/Programs/UnrealBuildTool/Configuration/ReadOnlyTargetRules.cs
+++ b/Engine/Source/Programs/UnrealBuildTool/Configuration/ReadOnlyTargetRules.cs
@@ -209,6 +209,8 @@ namespace UnrealBuildTool
 
 		public bool bEnableTrace => Inner.bEnableTrace;
 
+		public bool bForceEnableTrace => Inner.bForceEnableTrace;
+
 		public bool bCompileAgainstEditor => Inner.bCompileAgainstEditor;
 
 		public bool bCompileRecast => Inner.bCompileRecast;
diff --git a/Engine/Source/Programs/UnrealBuildTool/Configuration/TargetRules.cs b/Engine/Source/Programs/UnrealBuildTool/Configuration/TargetRules.cs
--- a/Engine/Source/Programs/UnrealBuildTool/Configuration/TargetRules.cs
+++ b/Engine/Source/Programs/UnrealBuildTool/Configuration/TargetRules.cs
@@ -1195,6 +1195,25 @@ namespace UnrealBuildTool
 		}
 		private bool? bEnableTracePrivate;
 
+		/// <summary>
+		/// Force enable tracing - used to allow testing programs to verify tracing works as intended.
+		/// </summary>
+		[RequiresUniqueBuildEnvironment]
+		public virtual bool bForceEnableTrace
+		{
+			get
+			{
+				if (bForceEnableTracePrivate.HasValue)
+				{
+					return bForceEnableTracePrivate.Value;
+				}
+				
+				return false;
+			}
+			set => bForceEnableTracePrivate = value;
+		}
+		private bool? bForceEnableTracePrivate;
+
 		/// <summary>
 		/// Manually specified value for bCompileAgainstEditor.
 		/// </summary>
diff --git a/Engine/Source/Programs/UnrealBuildTool/Configuration/UEBuildTarget.cs b/Engine/Source/Programs/UnrealBuildTool/Configuration/UEBuildTarget.cs
--- a/Engine/Source/Programs/UnrealBuildTool/Configuration/UEBuildTarget.cs
+++ b/Engine/Source/Programs/UnrealBuildTool/Configuration/UEBuildTarget.cs
@@ -5328,6 +5328,15 @@ namespace UnrealBuildTool
 				GlobalCompileEnvironment.Definitions.Add("UE_TRACE_ENABLED=0");
 			}
 
+			if (Rules.bForceEnableTrace)
+			{
+				GlobalCompileEnvironment.Definitions.Add("UE_TRACE_FORCE_ENABLED=1");
+			}
+			else
+			{
+				GlobalCompileEnvironment.Definitions.Add("UE_TRACE_FORCE_ENABLED=0");
+			}
+
 			// TODO - Sooner or later these should be removed.  bUseVerse has already been removed
 			GlobalCompileEnvironment.Definitions.Add("WITH_VERSE=1");
 			GlobalCompileEnvironment.Definitions.Add("UE_USE_VERSE_PATHS=1");
diff --git a/Engine/Source/Runtime/Core/Private/AutoRTFM/Context.cpp b/Engine/Source/Runtime/Core/Private/AutoRTFM/Context.cpp
--- a/Engine/Source/Runtime/Core/Private/AutoRTFM/Context.cpp
+++ b/Engine/Source/Runtime/Core/Private/AutoRTFM/Context.cpp
@@ -33,7 +33,36 @@ FAutoRTFMMetrics GetAutoRTFMMetrics()
 	return GAutoRTFMMetrics;
 }
 
-thread_local TUniquePtr<FContext> ContextTls;
+// A memory blob that can placement new store an FContext class.
+struct alignas(alignof(FContext)) FContextData final
+{
+	uint8 Payload[sizeof(FContext)] = {};
+};
+
+thread_local FContextData ContextDataTls;
+
+// This holder's entire job is to call the destructor on the memory we placement newed
+// into `ContextDataTls` above, and then zero out the data. This is to prevent the
+// problem where global destruction order is undefined, and we could call a destructor
+// that wants to call into the runtime after we have destroyed it. By zeroing out the
+// memory the runtime will *always* return safe defaults if it is queried after these
+// values are destroyed.
+struct FContextHolder final
+{
+	~FContextHolder()
+	{
+		if (Context)
+		{
+			Context->~FContext();
+			memset(Context, 0, sizeof(FContext));
+			Context = nullptr;
+		}
+	}
+
+	FContext* Context = nullptr;
+};
+
+thread_local FContextHolder ContextTls;
 
 void FContext::InitializeGlobalData()
 {
@@ -41,12 +70,12 @@ void FContext::InitializeGlobalData()
 
 FContext* FContext::TryGet()
 {
-    return ContextTls.Get();
+    return ContextTls.Context;
 }
 
 void FContext::Set()
 {
-    ContextTls.Reset(this);
+    ContextTls.Context = this;
 }
 
 FContext* FContext::Get()
@@ -55,7 +84,7 @@ FContext* FContext::Get()
 
     if (!Result)
     {
-        Result = new FContext();
+        Result = new (&ContextDataTls) FContext();
         Result->Set();
     }
 
diff --git a/Engine/Source/Runtime/Engine/Private/Tests/AutoRTFM/AutoRTFMTestObject.h b/Engine/Source/Runtime/Engine/Private/Tests/AutoRTFM/AutoRTFMTestObject.h
new file mode 100644
--- /dev/null
+++ b/Engine/Source/Runtime/Engine/Private/Tests/AutoRTFM/AutoRTFMTestObject.h
@@ -0,0 +1,26 @@
+// Copyright Epic Games, Inc. All Rights Reserved.
+
+#pragma once
+
+#include "AutoRTFMTestObject.generated.h"
+
+UCLASS()
+class UAutoRTFMTestObject : public UObject
+{
+    GENERATED_BODY()
+
+public:
+    UAutoRTFMTestObject(const FObjectInitializer& ObjectInitializer = FObjectInitializer::Get()) : Value(42)
+    {
+        UObject* const Obj = ObjectInitializer.GetObj();
+        UObject* const Outer = Obj->GetOuter();
+
+        if (Outer->IsA<UAutoRTFMTestObject>())
+        {
+            UAutoRTFMTestObject* const OuterAsType = static_cast<UAutoRTFMTestObject*>(Outer);
+            OuterAsType->Value += 13;
+        }
+    }
+
+    int Value;
+};
diff --git a/Engine/Source/Runtime/Engine/Private/Tests/AutoRTFM/AutoRTFMTraceFilterTests.cpp b/Engine/Source/Runtime/Engine/Private/Tests/AutoRTFM/AutoRTFMTraceFilterTests.cpp
new file mode 100644
--- /dev/null
+++ b/Engine/Source/Runtime/Engine/Private/Tests/AutoRTFM/AutoRTFMTraceFilterTests.cpp
@@ -0,0 +1,157 @@
+// Copyright Epic Games, Inc. All Rights Reserved.
+
+#include "Misc/AutomationTest.h"
+#include "Misc/App.h"
+#include "TraceFilter.h"
+#include "AutoRTFM/AutoRTFM.h"
+#include "AutoRTFMTestObject.h"
+
+#if WITH_DEV_AUTOMATION_TESTS
+
+#define TEST_CHECK_TRUE(b) do \
+{ \
+	if (!(b)) \
+	{ \
+		FString String(FString::Printf(TEXT("FAILED: %s:%u"), TEXT(__FILE__), __LINE__)); \
+		ExecutionInfo.AddEvent(FAutomationEvent(EAutomationEventType::Info, String)); \
+		return false; \
+	} \
+} while(false)
+
+IMPLEMENT_SIMPLE_AUTOMATION_TEST(FAutoRTFMTraceFilterTests, "AutoRTFM + FTraceFilter", EAutomationTestFlags::EngineFilter | EAutomationTestFlags::ClientContext | EAutomationTestFlags::ServerContext | EAutomationTestFlags::CommandletContext)
+
+bool FAutoRTFMTraceFilterTests::RunTest(const FString & Parameters)
+{
+	if (!AutoRTFM::ForTheRuntime::IsAutoRTFMRuntimeEnabled())
+	{
+		ExecutionInfo.AddEvent(FAutomationEvent(EAutomationEventType::Info, TEXT("SKIPPED 'FAutoRTFMTraceFilterTests' test. AutoRTFM disabled.")));
+		return true;
+	}
+
+#if TRACE_FILTERING_ENABLED
+	{
+		UAutoRTFMTestObject* Object = NewObject<UAutoRTFMTestObject>();
+
+		FTraceFilter::SetObjectIsTraceable(Object, false);
+		TEST_CHECK_TRUE(!FTraceFilter::IsObjectTraceable(Object));
+
+		AutoRTFM::ETransactionResult Result = AutoRTFM::Transact([&]
+			{
+				FTraceFilter::SetObjectIsTraceable(Object, true);
+				AutoRTFM::AbortTransaction();
+			});
+
+		TEST_CHECK_TRUE(AutoRTFM::ETransactionResult::AbortedByRequest == Result);
+		TEST_CHECK_TRUE(!FTraceFilter::IsObjectTraceable(Object));
+
+		Result = AutoRTFM::Transact([&]
+			{
+				FTraceFilter::SetObjectIsTraceable(Object, true);
+			});
+
+		TEST_CHECK_TRUE(AutoRTFM::ETransactionResult::Committed == Result);
+		TEST_CHECK_TRUE(FTraceFilter::IsObjectTraceable(Object));
+
+		Result = AutoRTFM::Transact([&]
+			{
+				AutoRTFM::OnAbort([Object]
+					{
+						FTraceFilter::SetObjectIsTraceable(Object, false);
+					});
+
+				AutoRTFM::AbortTransaction();
+			});
+
+		TEST_CHECK_TRUE(!FTraceFilter::IsObjectTraceable(Object));
+
+		Result = AutoRTFM::Transact([&]
+			{
+				AutoRTFM::OnCommit([Object]
+					{
+						FTraceFilter::SetObjectIsTraceable(Object, true);
+					});
+			});
+
+		TEST_CHECK_TRUE(FTraceFilter::IsObjectTraceable(Object));
+	}
+
+	{
+		UAutoRTFMTestObject* Object = NewObject<UAutoRTFMTestObject>();
+
+		FTraceFilter::SetObjectIsTraceable(Object, false);
+		TEST_CHECK_TRUE(!FTraceFilter::IsObjectTraceable(Object));
+
+		FTraceFilter::MarkObjectTraceable(Object);
+		TEST_CHECK_TRUE(FTraceFilter::IsObjectTraceable(Object));
+
+		FTraceFilter::SetObjectIsTraceable(Object, false);
+		TEST_CHECK_TRUE(!FTraceFilter::IsObjectTraceable(Object));
+
+		AutoRTFM::ETransactionResult Result = AutoRTFM::Transact([&]
+			{
+				FTraceFilter::MarkObjectTraceable(Object);
+				AutoRTFM::AbortTransaction();
+			});
+
+		TEST_CHECK_TRUE(AutoRTFM::ETransactionResult::AbortedByRequest == Result);
+		TEST_CHECK_TRUE(!FTraceFilter::IsObjectTraceable(Object));
+
+		Result = AutoRTFM::Transact([&]
+			{
+				FTraceFilter::MarkObjectTraceable(Object);
+			});
+
+		TEST_CHECK_TRUE(AutoRTFM::ETransactionResult::Committed == Result);
+		TEST_CHECK_TRUE(FTraceFilter::IsObjectTraceable(Object));
+
+		FTraceFilter::SetObjectIsTraceable(Object, false);
+		TEST_CHECK_TRUE(!FTraceFilter::IsObjectTraceable(Object));
+
+		Result = AutoRTFM::Transact([&]
+			{
+				AutoRTFM::OnAbort([Object]
+					{
+						FTraceFilter::MarkObjectTraceable(Object);
+					});
+
+				AutoRTFM::AbortTransaction();
+			});
+
+		TEST_CHECK_TRUE(FTraceFilter::IsObjectTraceable(Object));
+
+		FTraceFilter::SetObjectIsTraceable(Object, false);
+		TEST_CHECK_TRUE(!FTraceFilter::IsObjectTraceable(Object));
+
+		Result = AutoRTFM::Transact([&]
+			{
+				AutoRTFM::OnCommit([Object]
+					{
+						FTraceFilter::MarkObjectTraceable(Object);
+					});
+			});
+
+		TEST_CHECK_TRUE(FTraceFilter::IsObjectTraceable(Object));
+	}
+
+	{
+		UAutoRTFMTestObject* Object = NewObject<UAutoRTFMTestObject>();
+
+		bool bTraceable = true;
+
+		AutoRTFM::ETransactionResult Result = AutoRTFM::Transact([&]
+			{
+				bTraceable = FTraceFilter::IsObjectTraceable(Object);
+			});
+
+		TEST_CHECK_TRUE(AutoRTFM::ETransactionResult::Committed == Result);
+		TEST_CHECK_TRUE(FTraceFilter::IsObjectTraceable(Object) == bTraceable);
+	}
+
+	return true;
+#else
+	ExecutionInfo.AddEvent(FAutomationEvent(EAutomationEventType::Info, TEXT("SKIPPED 'FAutoRTFMTraceFilterTests' test. Trace filtering disabled.")));
+	return true;
+#endif
+}
+
+#endif //WITH_DEV_AUTOMATION_TESTS
diff --git a/Engine/Source/Runtime/Engine/Private/TraceFilter.cpp b/Engine/Source/Runtime/Engine/Private/TraceFilter.cpp
--- a/Engine/Source/Runtime/Engine/Private/TraceFilter.cpp
+++ b/Engine/Source/Runtime/Engine/Private/TraceFilter.cpp
@@ -8,6 +8,116 @@
 
 #include "ObjectTrace.h"
 
+// A transactionally safe critical section that works in the following novel ways:
+// - In the open (non-transactional):
+//   - Take the lock like before. Simple!
+//   - Free the lock like before too.
+// - In the closed (transactional):
+//   - During locking we query `TransactionalLockCount`:
+//	   - 0 means we haven't taken the lock within our transaction nest and need to acquire the lock.
+//     - Otherwise we already have the lock (and are preventing non-transactional code seeing any
+//       modifications we've made while holding the lock), so just bump `TransactionalLockCount`.
+//     - We also register an on-abort handler to release the lock should we abort (but we need to
+//       query `TransactionalLockCount` even there because we could be aborting an inner transaction
+//       and the parent transaction still wants to have the lock held!).
+//   - During unlocking we defer doing the unlock until the transaction commits.
+//
+// Thus with this approach we will hold this lock for the *entirety* of the transactional nest should
+// we take the lock during the transaction, thus preventing non-transactional code from seeing any
+// modifications we should make.
+struct FTransactionallySafeCriticalSection final
+{
+	FORCEINLINE void Lock()
+	{
+		if (AutoRTFM::IsTransactional())
+		{
+			AutoRTFM::Open([&]
+				{
+					// The transactional system which can increment TransactionalLockCount
+					// is always single-threaded, thus this is safe to check without atomicity.
+					if (0 == TransactionalLockCount)
+					{
+						CriticalSection.Lock();
+					}
+
+					TransactionalLockCount += 1;
+				});
+			
+			AutoRTFM::OnAbort([this]
+				{
+					check(0 != TransactionalLockCount);
+					TransactionalLockCount -= 1;
+
+					if (0 == TransactionalLockCount)
+					{
+						CriticalSection.Unlock();
+					}
+				});
+		}
+		else
+		{
+			CriticalSection.Lock();
+			check(0 == TransactionalLockCount);
+		}
+	}
+
+	void Unlock()
+	{
+		if (AutoRTFM::IsTransactional())
+		{
+			AutoRTFM::OnCommit([this]
+				{
+					check(0 != TransactionalLockCount);
+					TransactionalLockCount -= 1;
+
+					if (0 == TransactionalLockCount)
+					{
+						CriticalSection.Unlock();
+					}
+				});
+		}
+		else
+		{
+			check(0 == TransactionalLockCount);
+			CriticalSection.Unlock();
+		}
+	}
+
+private:
+	FCriticalSection CriticalSection;
+	uint32 TransactionalLockCount = 0;
+};
+
+struct FTransactionallySafeScopeLock final
+{
+	UE_NODISCARD_CTOR FTransactionallySafeScopeLock(FTransactionallySafeCriticalSection* InSynchObject)
+		: SynchObject(InSynchObject)
+	{
+		check(SynchObject);
+		SynchObject->Lock();
+	}
+
+	~FTransactionallySafeScopeLock()
+	{
+		Unlock();
+	}
+
+	void Unlock()
+	{
+		if (SynchObject)
+		{
+			SynchObject->Unlock();
+			SynchObject = nullptr;
+		}
+	}
+
+private:
+	FTransactionallySafeScopeLock() = delete;
+	FTransactionallySafeScopeLock(const FTransactionallySafeScopeLock&) = delete;
+
+	FTransactionallySafeCriticalSection* SynchObject;
+};
+
 struct ENGINE_API FTraceFilterObjectAnnotation
 {
 	FTraceFilterObjectAnnotation()
@@ -56,7 +166,7 @@ private:
 	void AddAnnotationInternal(const UObjectBase* Object, const FTraceFilterObjectAnnotation& Annotation)
 	{
 		check(Object);
-		FScopeLock ScopeLock(&AnnotationMapCritical);
+		FTransactionallySafeScopeLock ScopeLock(&AnnotationMapCritical);
 		AnnotationCacheKey = Object;
 		AnnotationCacheValue = Annotation;
 		if (AnnotationCacheValue.IsDefault())
@@ -88,7 +198,7 @@ public:
 
 	void RemoveAnnotation(const UObjectBase* Object)
 	{
-		FScopeLock ScopeLock(&AnnotationMapCritical);
+		FTransactionallySafeScopeLock ScopeLock(&AnnotationMapCritical);
 		check(Object);
 		
 		AnnotationCacheKey = Object;
@@ -98,13 +208,16 @@ public:
 		if (bHadElements && AnnotationMap.Num() == 0)
 		{
 			// we are removing the last one, so if we are auto removing or verifying removal, unregister now
-			GUObjectArray.RemoveUObjectDeleteListener(this);
+			AutoRTFM::OnCommit([this]
+				{
+					GUObjectArray.RemoveUObjectDeleteListener(this);
+				});
 		}
 	}
 
 	void RemoveAllAnnotations()
 	{
-		FScopeLock ScopeLock(&AnnotationMapCritical);
+		FTransactionallySafeScopeLock ScopeLock(&AnnotationMapCritical);
 
 		AnnotationCacheKey = NULL;
 		AnnotationCacheValue = FTraceFilterObjectAnnotation();
@@ -119,7 +232,7 @@ public:
 
 	FORCEINLINE FTraceFilterObjectAnnotation GetAnnotation(const UObjectBase* Object)
 	{
-		FScopeLock ScopeLock(&AnnotationMapCritical);
+		FTransactionallySafeScopeLock ScopeLock(&AnnotationMapCritical);
 
 		check(Object);
 		
@@ -146,7 +259,7 @@ public:
 
 	void Lock()
 	{
-		UniqueScopeLock = MakeUnique<FScopeLock>(&AnnotationMapCritical);
+		UniqueScopeLock = MakeUnique<FTransactionallySafeScopeLock>(&AnnotationMapCritical);
 	}
 
 	void Unlock()
@@ -166,9 +279,9 @@ public:
 
 private:
 	TMap<const UObjectBase*, FTraceFilterObjectAnnotation> AnnotationMap;
-	FCriticalSection AnnotationMapCritical;
+	FTransactionallySafeCriticalSection AnnotationMapCritical;
 
-	TUniquePtr<FScopeLock> UniqueScopeLock;
+	TUniquePtr<FTransactionallySafeScopeLock> UniqueScopeLock;
 
 	const UObjectBase* AnnotationCacheKey;
 	FTraceFilterObjectAnnotation AnnotationCacheValue;
diff --git a/Engine/Source/Runtime/Engine/Public/TraceFilter.h b/Engine/Source/Runtime/Engine/Public/TraceFilter.h
--- a/Engine/Source/Runtime/Engine/Public/TraceFilter.h
+++ b/Engine/Source/Runtime/Engine/Public/TraceFilter.h
@@ -13,8 +13,7 @@
 #include "Templates/UnrealTemplate.h"
 #include "UObject/Object.h"
 
-
-#if UE_TRACE_ENABLED && !IS_PROGRAM && !UE_BUILD_SHIPPING
+#if UE_TRACE_FORCE_ENABLED || (UE_TRACE_ENABLED && !IS_PROGRAM && !UE_BUILD_SHIPPING)
 #define TRACE_FILTERING_ENABLED 1
 #else
 #define TRACE_FILTERING_ENABLED 0
