Fix the compiler warning with the new AutoRTFMEngineTests.

diff --git a/Engine/Source/Programs/AutoRTFMEngineTests/AutoRTFMEngineTests.Build.cs b/Engine/Source/Programs/AutoRTFMEngineTests/AutoRTFMEngineTests.Build.cs
--- a/Engine/Source/Programs/AutoRTFMEngineTests/AutoRTFMEngineTests.Build.cs
+++ b/Engine/Source/Programs/AutoRTFMEngineTests/AutoRTFMEngineTests.Build.cs
@@ -26,7 +26,6 @@ public class AutoRTFMEngineTests : ModuleRules
 				"MoviePlayerProxy",
 				"PreLoadScreen",
 				"ProfilerService",
-				"ReplicationSystemTestPlugin",
 				"SessionServices",
 				"SlateNullRenderer",
 				"SlateRHIRenderer",
diff --git a/Engine/Source/Programs/AutoRTFMEngineTests/AutoRTFMEngineTests.Target.cs b/Engine/Source/Programs/AutoRTFMEngineTests/AutoRTFMEngineTests.Target.cs
--- a/Engine/Source/Programs/AutoRTFMEngineTests/AutoRTFMEngineTests.Target.cs
+++ b/Engine/Source/Programs/AutoRTFMEngineTests/AutoRTFMEngineTests.Target.cs
@@ -16,7 +16,7 @@ public class AutoRTFMEngineTestsTarget : TargetRules
 		bBuildWithEditorOnlyData = false;
 
 		bCompileAgainstEngine = true;
-		bCompileAgainstCoreUObject = false;
+		bCompileAgainstCoreUObject = true;
 		bCompileAgainstApplicationCore = true;
         bCompileWithPluginSupport = true;
         bBuildDeveloperTools = false;
diff --git a/Engine/Source/Programs/AutoRTFMEngineTests/Private/AutoRTFMEngineTests.cpp b/Engine/Source/Programs/AutoRTFMEngineTests/Private/AutoRTFMEngineTests.cpp
--- a/Engine/Source/Programs/AutoRTFMEngineTests/Private/AutoRTFMEngineTests.cpp
+++ b/Engine/Source/Programs/AutoRTFMEngineTests/Private/AutoRTFMEngineTests.cpp
@@ -21,7 +21,7 @@ typedef FNullTestRunner TestRunner;
 
 DEFINE_LOG_CATEGORY(LogAutoRTFMEngineTests);
 
-IMPLEMENT_APPLICATION(AutoRTFMEngineTests, "");
+IMPLEMENT_APPLICATION(AutoRTFMEngineTests, "AutoRTFMEngineTests");
 
 static void PreInit();
 static void LoadModules();
@@ -130,21 +130,6 @@ static void PreInit()
 #endif
 	FPlatformMemory::Init();
 
-	// Initialize trace
-	FString Parameter;
-	if (FParse::Value(FCommandLine::Get(), TEXT("-trace="), Parameter, false))
-	{
-		FTraceAuxiliary::Initialize(FCommandLine::Get());
-		FTraceAuxiliary::TryAutoConnect();
-	}
-
-	if (IPlatformFile* WrapperFile = FPlatformFileManager::Get().GetPlatformFile(TEXT("AutoRTFMEngineTestsFile")))
-	{
-		IPlatformFile* CurrentPlatformFile = &FPlatformFileManager::Get().GetPlatformFile();
-		WrapperFile->Initialize(CurrentPlatformFile, TEXT(""));
-		FPlatformFileManager::Get().SetPlatformFile(*WrapperFile);
-	}
-
 #if WITH_COREUOBJECT
 	// Initialize the PackageResourceManager, which is needed to load any (non-script) Packages. It is first used in ProcessNewlyLoadedObjects (due to the loading of asset references in Class Default Objects)
 	// It has to be intialized after the AssetRegistryModule; the editor implementations of PackageResourceManager relies on it
@@ -157,6 +142,9 @@ static void PreInit()
 
 	// Config overrides
 	GConfig->SetInt(TEXT("/Script/Engine.GarbageCollectionSettings"), TEXT("gc.MaxObjectsNotConsideredByGC"), 0, GEngineIni);
+	GConfig->SetInt(TEXT("/Script/Engine.GarbageCollectionSettings"), TEXT("gc.MaxObjectsInProgram"), 500000, GEngineIni);
+	GConfig->SetInt(TEXT("/Script/Engine.GarbageCollectionSettings"), TEXT("gc.MaxObjectsInGame"), 500000, GEngineIni);
+	GConfig->SetInt(TEXT("/Script/Engine.GarbageCollectionSettings"), TEXT("gc.MaxObjectsInEditor"), 500000, GEngineIni);
 	GConfig->SetString(TEXT("/Script/Engine.Engine"), TEXT("AIControllerClassName"), TEXT("/Script/AIModule.AIController"), GEngineIni);
 	GConfig->SetString(TEXT("/Script/Engine.Engine"), TEXT("DefaultMaterialName"), TEXT("/Engine/Transient.MockDefaultMaterial"), GEngineIni);
 	GConfig->SetString(TEXT("/Script/Engine.Engine"), TEXT("DefaultLightFunctionMaterialName"), TEXT("/Engine/Transient.MockDefaultMaterial"), GEngineIni);
@@ -175,8 +163,6 @@ static void PreInit()
 
 	FDelayedAutoRegisterHelper::RunAndClearDelayedAutoRegisterDelegates(EDelayedRegisterRunPhase::StatSystemReady);
 
-	FCoreStyle::ResetToDefault();
-	FUMGCoreStyle::ResetToDefault();
 }
 
 static void LoadModules()
@@ -190,6 +176,9 @@ static void LoadModules()
 	FCoreDelegates::OnInit.Broadcast();
 #endif
 
+	FCoreStyle::ResetToDefault();
+	FUMGCoreStyle::ResetToDefault();
+
 	// Create a mock default material to keep the material system happy
 	UMaterial* MockMaterial = NewObject<UMaterial>(GetTransientPackage(), UMaterial::StaticClass(), TEXT("MockDefaultMaterial"), RF_Transient | RF_MarkAsRootSet);
 
diff --git a/Engine/Source/Programs/AutoRTFMEngineTests/Private/AutoRTFMEngineTestsFileModule.cpp b/Engine/Source/Programs/AutoRTFMEngineTests/Private/AutoRTFMEngineTestsFileModule.cpp
deleted file mode 100644
