Disable AutoRTFM compiler needed when doing GPF

diff --git a/Engine/Source/Programs/AutoRTFMEngineTests/AutoRTFMEngineTests.Target.cs b/Engine/Source/Programs/AutoRTFMEngineTests/AutoRTFMEngineTests.Target.cs
--- a/Engine/Source/Programs/AutoRTFMEngineTests/AutoRTFMEngineTests.Target.cs
+++ b/Engine/Source/Programs/AutoRTFMEngineTests/AutoRTFMEngineTests.Target.cs
@@ -45,7 +45,10 @@ public class AutoRTFMEngineTestsTarget : TargetRules
 		// Load time profiling brings object construction to a crawl.
 		GlobalDefinitions.Add("LOADTIMEPROFILERTRACE_ENABLED=0");
 
-		bUseAutoRTFMCompiler = true;
+		if (!bGenerateProjectFiles)
+		{
+			bUseAutoRTFMCompiler = true;
+		}
 	}
 }
 
