[AutoRTFMEngineTests] Add a workaround for loading of editor-only data and re-enable presubmits

Multiple transitive dependency engine components will attempt to load `.uasset` files out of the `/Engine` directory.
These files can only be loaded if `WITH_EDITORONLY_DATA` is enabled (`bBuildWithEditorOnlyData`), but this is currently _very_ difficult to turn on without build issues.

As none of our tests care about these assets (most are textures), add a hack that redirects the file loads to a non-existent file.

This allows us to re-enable the presubmits for the AutoRTFMEngineTests.

Also fully-qualify the AutoRTFM namespace from the IMPL2 macros. The unity builds can fail to resolve if you touch certain files.

diff --git a/Engine/Source/Programs/AutoRTFMEngineTests/Private/AutoRTFMEngineTests.cpp b/Engine/Source/Programs/AutoRTFMEngineTests/Private/AutoRTFMEngineTests.cpp
--- a/Engine/Source/Programs/AutoRTFMEngineTests/Private/AutoRTFMEngineTests.cpp
+++ b/Engine/Source/Programs/AutoRTFMEngineTests/Private/AutoRTFMEngineTests.cpp
@@ -8,6 +8,7 @@
 #include "Materials/Material.h"
 #include "AutoRTFM/AutoRTFM.h"
 #include "AutoRTFMTestEngine.h"
+#include "UObject/CoreRedirects.h"
 
 #if WITH_AUTOMATION_WORKER
 namespace UE::AutoRTFM
@@ -39,7 +40,7 @@ INT32_MAIN_INT32_ARGC_TCHAR_ARGV()
 
 #if !(defined(PLATFORM_XBOXONE) && PLATFORM_XBOXONE)
 	// Parse original cmdline if there is one
-	OriginalCmdLine = FCommandLine::BuildFromArgV(nullptr, ArgC, ArgV, nullptr);	
+	OriginalCmdLine = FCommandLine::BuildFromArgV(nullptr, ArgC, ArgV, nullptr);
 #endif
 
 	// Due to some code not respecting nullrhi etc.
@@ -58,7 +59,7 @@ INT32_MAIN_INT32_ARGC_TCHAR_ARGV()
 	PreInit();
 	LoadModules();
 	PostInit();
-	
+
 	GEngine = NewObject<UAutoRTFMTestEngine>(GetTransientPackage());
 	GEngine->DefaultPhysMaterial = NewObject<UPhysicalMaterial>();
 
@@ -79,7 +80,7 @@ INT32_MAIN_INT32_ARGC_TCHAR_ARGV()
 		{
 			TestsPassed = Runner->RunTests();
 		}
-		
+
 	}
 
 	TearDown();
@@ -118,8 +119,53 @@ static void TearDown()
 	FTraceAuxiliary::Shutdown();
 }
 
+// Adds redirects to a non-existent file for all the Engine .uasset files that
+// will be automatically loaded by the various engine systems. This is done to
+// prevent attempted deserialization of assets that can only be loaded when
+// the project is built with WITH_EDITORONLY_DATA. Unfortunately turning this
+// flag on also requires WITH_EDITOR which is currently extremely difficult to
+// build outside of the editor.
+// HACK: SOL-6723
+static void PreventLoadingOfEditorOnlyData()
+{
+	const TCHAR* const IncompatiblePackages[] = {
+		TEXT("/Engine/EngineResources/DefaultTexture"),
+		TEXT("/Engine/EngineResources/DefaultTextureCube"),
+		TEXT("/Engine/EngineResources/DefaultVolumeTexture"),
+		TEXT("/Engine/EngineFonts/RobotoDistanceField"),
+		TEXT("/Engine/EngineMaterials/DefaultTextMaterialOpaque"),
+		TEXT("/Engine/EngineDamageTypes/DmgTypeBP_Environmental"),
+		TEXT("/Engine/EngineSky/VolumetricClouds/m_SimpleVolumetricCloud_Inst"),
+		TEXT("/Engine/EngineMeshes/Sphere"),
+		TEXT("/Engine/EngineResources/WhiteSquareTexture"),
+		TEXT("/Engine/EngineResources/GradientTexture0"),
+		TEXT("/Engine/EngineResources/Black"),
+		TEXT("/Engine/EngineDebugMaterials/VolumeToRender"),
+		TEXT("/Engine/EngineDebugMaterials/M_VolumeRenderSphereTracePP"),
+		TEXT("/Engine/EngineFonts/Roboto"),
+		TEXT("/Engine/EngineMaterials/Widget3DPassThrough_Translucent"),
+		TEXT("/Engine/EngineMaterials/Widget3DPassThrough_Translucent_OneSided"),
+		TEXT("/Engine/EngineMaterials/Widget3DPassThrough_Opaque"),
+		TEXT("/Engine/EngineMaterials/Widget3DPassThrough_Opaque_OneSided"),
+		TEXT("/Engine/EngineMaterials/Widget3DPassThrough_Masked"),
+		TEXT("/Engine/EngineMaterials/Widget3DPassThrough_Masked_OneSided")
+	};
+	FCoreRedirectObjectName InvalidName(NAME_None, NAME_None, TEXT("/Engine/DoesNotExist"));
+	TArray<FCoreRedirect> NewRedirects;
+	for (const TCHAR* PackageName : IncompatiblePackages)
+	{
+		NewRedirects.Emplace(ECoreRedirectFlags::Type_Package,
+			FCoreRedirectObjectName(NAME_None, NAME_None, PackageName),
+			InvalidName);
+	}
+	FCoreRedirects::AddRedirectList(NewRedirects, TEXT("AutoRTFMEngineTests.PreventLoadingOfEditorOnlyData"));
+	FCoreRedirects::AddKnownMissing(ECoreRedirectFlags::Type_Package, InvalidName);
+}
+
 static void PreInit()
 {
+	PreventLoadingOfEditorOnlyData();
+
 	// We enable the AutoRTFM runtime as the tests depend on it.
 	AutoRTFM::ForTheRuntime::SetAutoRTFMRuntime(AutoRTFM::ForTheRuntime::AutoRTFM_Enabled);
 
diff --git a/Engine/Source/Runtime/Core/Public/AutoRTFM/AutoRTFM.h b/Engine/Source/Runtime/Core/Public/AutoRTFM/AutoRTFM.h
--- a/Engine/Source/Runtime/Core/Public/AutoRTFM/AutoRTFM.h
+++ b/Engine/Source/Runtime/Core/Public/AutoRTFM/AutoRTFM.h
@@ -970,10 +970,10 @@ namespace AutoRTFM::Private
 #define UE_AUTORTFM_ONCOMMIT_IMPL(...) UE_AUTORTFM_BEGIN_DISABLE_WARNINGS ::AutoRTFM::OnCommit([=]() { __VA_ARGS__ }) UE_AUTORTFM_END_DISABLE_WARNINGS
 #define UE_AUTORTFM_TRANSACT_IMPL(...) ::AutoRTFM::Transact([&]() { __VA_ARGS__ })
 
-#define UE_AUTORTFM_OPEN_IMPL2 AutoRTFM::Private::FOpenHelper{} + [&]()
-#define UE_AUTORTFM_ONABORT_IMPL2(...) AutoRTFM::Private::FOnAbortHelper{} + [__VA_ARGS__]()
-#define UE_AUTORTFM_ONCOMMIT_IMPL2(...) AutoRTFM::Private::FOnCommitHelper{} + [__VA_ARGS__]()
-#define UE_AUTORTFM_TRANSACT_IMPL2 AutoRTFM::Private::FTransactHelper{} + [&]() { __VA_ARGS__ })
+#define UE_AUTORTFM_OPEN_IMPL2 ::AutoRTFM::Private::FOpenHelper{} + [&]()
+#define UE_AUTORTFM_ONABORT_IMPL2(...) ::AutoRTFM::Private::FOnAbortHelper{} + [__VA_ARGS__]()
+#define UE_AUTORTFM_ONCOMMIT_IMPL2(...) ::AutoRTFM::Private::FOnCommitHelper{} + [__VA_ARGS__]()
+#define UE_AUTORTFM_TRANSACT_IMPL2 ::AutoRTFM::Private::FTransactHelper{} + [&]() { __VA_ARGS__ })
 #else
 
 // Older macros where the code is passed as a macro argument. These should be phased out as they make debugging more difficult
@@ -983,7 +983,7 @@ namespace AutoRTFM::Private
 #define UE_AUTORTFM_TRANSACT_IMPL(...) do { __VA_ARGS__ } while (false)
 
 // Do nothing, these should be followed by blocks that should be either executed or not executed
-#define UE_AUTORTFM_OPEN_IMPL2 
+#define UE_AUTORTFM_OPEN_IMPL2
 #define UE_AUTORTFM_ONABORT_IMPL2(...) while (false)
 #define UE_AUTORTFM_ONCOMMIT_IMPL2(...)
 #define UE_AUTORTFM_TRANSACT_IMPL2
@@ -991,7 +991,7 @@ namespace AutoRTFM::Private
 
 // Runs a block of code in the open, non-transactionally. Anything performed in the open will not be undone if a transaction fails.
 #define UE_AUTORTFM_OPEN(...) UE_AUTORTFM_OPEN_IMPL(__VA_ARGS__)
-// This new version is used like UE_AUTORTFM_OPEN2 { ... code ... }; 
+// This new version is used like UE_AUTORTFM_OPEN2 { ... code ... };
 #define UE_AUTORTFM_OPEN2 UE_AUTORTFM_OPEN_IMPL2
 
 // Runs a block of code if a transaction aborts.
