[AutoRTFMEngineTests] Add '-NoZenLoader' to the command line argument

The new 'ZenLoader' attempts to load the packages before redirects, which breaks the PreventLoadingOfEditorOnlyData() bodge

diff --git a/Engine/Source/Programs/AutoRTFMEngineTests/Private/AutoRTFMEngineTests.cpp b/Engine/Source/Programs/AutoRTFMEngineTests/Private/AutoRTFMEngineTests.cpp
--- a/Engine/Source/Programs/AutoRTFMEngineTests/Private/AutoRTFMEngineTests.cpp
+++ b/Engine/Source/Programs/AutoRTFMEngineTests/Private/AutoRTFMEngineTests.cpp
@@ -50,7 +50,25 @@ INT32_MAIN_INT32_ARGC_TCHAR_ARGV()
 
 	// Init cmd line used for test
 	{
-		FString CmdLineOverride(TEXT("-nullrhi -NoAsyncLoadingThread -NoAsyncPostLoad -noedl -unattended -LogCmds=\"LogSlate off, LogSlateStyle off, LogUObjectBase off, LogUObjectGlobals off, LogConsoleResponse off, LogPackageLocalizationManager off, LogStreaming off, LogCsvProfiler off, LogDeviceProfileManager off, LogConfig off, AutoRTFMEngineTests on\""));
+		FString CmdLineOverride(TEXT(
+				"-nullrhi "
+				"-NoAsyncLoadingThread "
+				"-NoAsyncPostLoad "
+				"-NoZenLoader " // Required for the PreventLoadingOfEditorOnlyData() bodge
+				"-noedl "
+				"-unattended "
+				"-LogCmds=\""
+					"LogSlate off, "
+					"LogSlateStyle off, "
+					"LogUObjectBase off, "
+					"LogUObjectGlobals off, "
+					"LogConsoleResponse off, "
+					"LogPackageLocalizationManager off, "
+					"LogStreaming off, "
+					"LogCsvProfiler off, "
+					"LogDeviceProfileManager off, "
+					"LogConfig off, "
+					"AutoRTFMEngineTests on\""));
 		FCommandLine::Set(ToCStr(CmdLineOverride));
 	}
 
