Set NOINITCRASHREPORTER=1 in the Verse test targets to improve process startup overhead, especially when running on WSL with slow file I/O.
Without the flag, every invocation creates a directory and some files in Engine/Programs/TargetName/Saved/Config/CrashReporterClient, stats all the existing directories there, and deletes any older than 2 days. When I discovered this, I had 16k directories there, and the sanitized tests running on WSL spent quite a bit of time on this file I/O before any of our code even ran.
With 16k crash report client configs, this makes the Verse tests that run on WSL ~2.8x faster.

diff --git a/Engine/Source/Programs/AutoRTFMEngineTests/AutoRTFMEngineTests.Target.cs b/Engine/Source/Programs/AutoRTFMEngineTests/AutoRTFMEngineTests.Target.cs
--- a/Engine/Source/Programs/AutoRTFMEngineTests/AutoRTFMEngineTests.Target.cs
+++ b/Engine/Source/Programs/AutoRTFMEngineTests/AutoRTFMEngineTests.Target.cs
@@ -47,6 +47,9 @@ public class AutoRTFMEngineTestsTarget : TargetRules
 
 		GlobalDefinitions.Add("UNIX_PLATFORM_FILE_SPEEDUP_FILE_OPERATIONS=1");
 
+		// Disable crashreporter to improve startup time.
+		GlobalDefinitions.Add("NOINITCRASHREPORTER=1");
+
 		if (!bGenerateProjectFiles)
 		{
 			bUseAutoRTFMCompiler = true;
diff --git a/Engine/Source/Programs/AutoRTFMTests/AutoRTFMTests.Target.cs b/Engine/Source/Programs/AutoRTFMTests/AutoRTFMTests.Target.cs
--- a/Engine/Source/Programs/AutoRTFMTests/AutoRTFMTests.Target.cs
+++ b/Engine/Source/Programs/AutoRTFMTests/AutoRTFMTests.Target.cs
@@ -66,6 +66,9 @@ public class AutoRTFMTestsTarget : TargetRules
 
 		GlobalDefinitions.Add("UNIX_PLATFORM_FILE_SPEEDUP_FILE_OPERATIONS=1");
 
+		// Disable crashreporter to improve startup time.
+		GlobalDefinitions.Add("NOINITCRASHREPORTER=1");
+
 		// For testing insights, this sets `UE_TRACE_ENABLED`.
 		bEnableTrace = true;
 	}
