[AutoRTFMEngineTests] Log test failures with the correct severity.

Test failures are near impossible to spot in logs, otherwise.

diff --git a/Engine/Source/Programs/AutoRTFMEngineTests/Private/AutomationTestRunner.cpp b/Engine/Source/Programs/AutoRTFMEngineTests/Private/AutomationTestRunner.cpp
--- a/Engine/Source/Programs/AutoRTFMEngineTests/Private/AutomationTestRunner.cpp
+++ b/Engine/Source/Programs/AutoRTFMEngineTests/Private/AutomationTestRunner.cpp
@@ -1,6 +1,8 @@
 // Copyright Epic Games, Inc. All Rights Reserved.
 
 #include "AutomationTestRunner.h"
+#include "Logging/LogVerbosity.h"
+#include "Misc/AutomationEvent.h"
 
 #if WITH_AUTOMATION_WORKER
 
@@ -60,7 +62,18 @@ FAutomationTestRunner::RunTests(const TCHAR* TestFilter)
 		{
 			for (const FAutomationExecutionEntry& Entry : ExecutionInfo.GetEntries())
 			{
-				UE_LOG(LogAutoRTFMEngineTests, Display, TEXT("%s"), *Entry.Event.Message);
+				switch(Entry.Event.Type)
+				{
+					case EAutomationEventType::Info:
+						UE_LOG(LogAutoRTFMEngineTests, Display, TEXT("%s"), *Entry.Event.Message);
+						break;
+					case EAutomationEventType::Warning:
+						UE_LOG(LogAutoRTFMEngineTests, Warning, TEXT("%s"), *Entry.Event.Message);
+						break;
+					case EAutomationEventType::Error:
+						UE_LOG(LogAutoRTFMEngineTests, Error, TEXT("%s"), *Entry.Event.Message);
+						break;
+				}
 			}
 
 			AllPassed = false;
