[AutoRTFM] Set GIsGameThreadIdInitialized when we initialize the game-thread ID.

Without this change, all threads return true from IsInGameThread().

This bug would cause the HTTP thread handler to fail when run from AutoRTFMEngineTests, regardless of AutoRTFM-enable status. This happened because FHttpThreadBase::UpdateThreadPriorityIfNeeded asserts that it is not being called from the game thread.

diff --git a/Engine/Source/Programs/AutoRTFMEngineTests/Private/AutoRTFMEngineTests.cpp b/Engine/Source/Programs/AutoRTFMEngineTests/Private/AutoRTFMEngineTests.cpp
--- a/Engine/Source/Programs/AutoRTFMEngineTests/Private/AutoRTFMEngineTests.cpp
+++ b/Engine/Source/Programs/AutoRTFMEngineTests/Private/AutoRTFMEngineTests.cpp
@@ -216,7 +216,7 @@ static void PreInit()
 
 #if WITH_COREUOBJECT
 	// Initialize the PackageResourceManager, which is needed to load any (non-script) Packages. It is first used in ProcessNewlyLoadedObjects (due to the loading of asset references in Class Default Objects)
-	// It has to be intialized after the AssetRegistryModule; the editor implementations of PackageResourceManager relies on it
+	// It has to be initialized after the AssetRegistryModule; the editor implementations of PackageResourceManager relies on it
 	IPackageResourceManager::Initialize();
 #endif
 
@@ -236,6 +236,8 @@ static void PreInit()
 	GConfig->SetString(TEXT("/Script/Engine.Engine"), TEXT("DefaultPostProcessMaterialName"), TEXT("/Engine/Transient.MockDefaultMaterial"), GEngineIni);
 
 	GGameThreadId = FPlatformTLS::GetCurrentThreadId();
+	GIsGameThreadIdInitialized = true;
+
 	FTaskGraphInterface::Startup(FPlatformMisc::NumberOfCores());
 	FTaskGraphInterface::Get().AttachToThread(ENamedThreads::GameThread);
 
