[AutoRTFM] Update "AutoRTFM/AutoRTFM.h" include to "AutoRTFM.h"

Now that AutoRTFM has been moved out of Core, the header is no longer in an AutoRTFM directory.

diff --git a/Engine/Source/Programs/AutoRTFMEngineTests/Private/AutoRTFMEngineTests.cpp b/Engine/Source/Programs/AutoRTFMEngineTests/Private/AutoRTFMEngineTests.cpp
--- a/Engine/Source/Programs/AutoRTFMEngineTests/Private/AutoRTFMEngineTests.cpp
+++ b/Engine/Source/Programs/AutoRTFMEngineTests/Private/AutoRTFMEngineTests.cpp
@@ -7,7 +7,7 @@
 #include "AutomationTestRunner.h"
 #include "NullTestRunner.h"
 #include "Materials/Material.h"
-#include "AutoRTFM/AutoRTFM.h"
+#include "AutoRTFM.h"
 #include "AutoRTFMTestEngine.h"
 #include "UObject/CoreRedirects.h"
 
