[AutoRTFM] Execute latent commands in the automation runner.

This step was missing compared to other implementations of the automation test runner. If any latent commands were registered, this would cause us to check-fail in TestFramework.StopTest. This doesn't affect any of our existing unit tests, but more complex test cases like entity tests rely on this step.

diff --git a/Engine/Source/Programs/AutoRTFMEngineTests/Private/AutomationTestRunner.cpp b/Engine/Source/Programs/AutoRTFMEngineTests/Private/AutomationTestRunner.cpp
--- a/Engine/Source/Programs/AutoRTFMEngineTests/Private/AutomationTestRunner.cpp
+++ b/Engine/Source/Programs/AutoRTFMEngineTests/Private/AutomationTestRunner.cpp
@@ -56,8 +56,9 @@ FAutomationTestRunner::RunTests(const TCHAR* TestFilter)
 
 		constexpr int32 RoleIndex = 0;
 		TestFramework.StartTestByName(TestInfo.GetTestName(), RoleIndex);
-		FAutomationTestExecutionInfo ExecutionInfo;
+        TestFramework.ExecuteLatentCommands();
 
+		FAutomationTestExecutionInfo ExecutionInfo;
 		if (!TestFramework.StopTest(ExecutionInfo))
 		{
 			for (const FAutomationExecutionEntry& Entry : ExecutionInfo.GetEntries())
