[AutoRTFM] Improve AutoRTFMEngineTests startup code.

While experimenting with the Entity framework tests, I made a handful of changes which helped AutoRTFMEngineTests start up a little more like a fully-initialized Engine. I decided against attempting to run Entity tests as they are too high-level, but these fixes should be harmless to land in isolation.

diff --git a/Engine/Source/Programs/AutoRTFMEngineTests/Classes/AutoRTFMTestEngine.h b/Engine/Source/Programs/AutoRTFMEngineTests/Classes/AutoRTFMTestEngine.h
--- a/Engine/Source/Programs/AutoRTFMEngineTests/Classes/AutoRTFMTestEngine.h
+++ b/Engine/Source/Programs/AutoRTFMEngineTests/Classes/AutoRTFMTestEngine.h
@@ -11,7 +11,12 @@ class UAutoRTFMTestEngine : public UEngine
     GENERATED_BODY()
 
 public:
-    int Value = 42;
+	void InitializeEngineSubsystemCollection()
+	{
+		EngineSubsystemCollection.Initialize(this);
+	}
+
+	int Value = 42;
 
 	void Tick(float, bool) override {}
 };
diff --git a/Engine/Source/Programs/AutoRTFMEngineTests/Private/AutoRTFMEngineTests.cpp b/Engine/Source/Programs/AutoRTFMEngineTests/Private/AutoRTFMEngineTests.cpp
--- a/Engine/Source/Programs/AutoRTFMEngineTests/Private/AutoRTFMEngineTests.cpp
+++ b/Engine/Source/Programs/AutoRTFMEngineTests/Private/AutoRTFMEngineTests.cpp
@@ -85,9 +85,11 @@ INT32_MAIN_INT32_ARGC_TCHAR_ARGV()
 	LoadModules();
 	PostInit();
 
-	GEngine = NewObject<UAutoRTFMTestEngine>(GetTransientPackage());
-	GEngine->DefaultPhysMaterial = NewObject<UPhysicalMaterial>();
-	GEngine->WorldSettingsClass = AWorldSettings::StaticClass();
+	UAutoRTFMTestEngine* TestEngine = NewObject<UAutoRTFMTestEngine>(GetTransientPackage());
+	TestEngine->DefaultPhysMaterial = NewObject<UPhysicalMaterial>();
+	TestEngine->WorldSettingsClass = AWorldSettings::StaticClass();
+	TestEngine->InitializeEngineSubsystemCollection();
+	GEngine = TestEngine;
 
 	FGenericDataDrivenShaderPlatformInfo::Initialize();
 
@@ -186,7 +188,10 @@ static void PreventLoadingOfEditorOnlyData()
 		TEXT("/Engine/EngineMaterials/Widget3DPassThrough_Opaque"),
 		TEXT("/Engine/EngineMaterials/Widget3DPassThrough_Opaque_OneSided"),
 		TEXT("/Engine/EngineMaterials/Widget3DPassThrough_Masked"),
-		TEXT("/Engine/EngineMaterials/Widget3DPassThrough_Masked_OneSided")
+		TEXT("/Engine/EngineMaterials/Widget3DPassThrough_Masked_OneSided"),
+		TEXT("/Engine/BasicShapes/Cube"),
+		TEXT("/CameraCalibrationCore/Materials/M_DefaultEvenSquare"),
+		TEXT("/CameraCalibrationCore/Materials/M_DefaultOddSquare"),
 	};
 	FCoreRedirectObjectName InvalidName(NAME_None, NAME_None, TEXT("/Engine/DoesNotExist"));
 	TArray<FCoreRedirect> NewRedirects;
