[AutoRTFM] Replace a confusing build failure with an explicit one.

The FNullTestRunner is not a usable replacement for the TestRunner in AutoRTFMEngineTests because its RunTests function returns void, not a bool, so it causes a compile error. We only disable the Automation Worker in Shipping builds, so now we now emit an explicit #error suggesting to build as Test instead.

diff --git a/Engine/Source/Programs/AutoRTFMEngineTests/Private/AutoRTFMEngineTests.cpp b/Engine/Source/Programs/AutoRTFMEngineTests/Private/AutoRTFMEngineTests.cpp
--- a/Engine/Source/Programs/AutoRTFMEngineTests/Private/AutoRTFMEngineTests.cpp
+++ b/Engine/Source/Programs/AutoRTFMEngineTests/Private/AutoRTFMEngineTests.cpp
@@ -19,7 +19,7 @@ namespace UE::AutoRTFM
 
 typedef UE::AutoRTFM::FAutomationTestRunner TestRunner;
 #else
-typedef FNullTestRunner TestRunner;
+#error Cannot run tests in a Shipping build! Use a Test build instead.
 #endif
 
 namespace AutoRTFM
